SkeinCLI
* Added copyright notice to source files
* Temporarily Removed ProgressBar as no TPS for it

diff --git a/Engine/Source/Programs/SkeinCLI/SkeinCLI/AssetsCmd.cs b/Engine/Source/Programs/SkeinCLI/SkeinCLI/AssetsCmd.cs
--- a/Engine/Source/Programs/SkeinCLI/SkeinCLI/AssetsCmd.cs
+++ b/Engine/Source/Programs/SkeinCLI/SkeinCLI/AssetsCmd.cs
@@ -1,4 +1,6 @@
-﻿using McMaster.Extensions.CommandLineUtils;
+﻿// Copyright Epic Games, Inc. All Rights Reserved.
+
+using McMaster.Extensions.CommandLineUtils;
 using Serilog;
 using System;
 using System.Collections.Generic;
diff --git a/Engine/Source/Programs/SkeinCLI/SkeinCLI/AuthCmd.cs b/Engine/Source/Programs/SkeinCLI/SkeinCLI/AuthCmd.cs
--- a/Engine/Source/Programs/SkeinCLI/SkeinCLI/AuthCmd.cs
+++ b/Engine/Source/Programs/SkeinCLI/SkeinCLI/AuthCmd.cs
@@ -1,4 +1,6 @@
-﻿using McMaster.Extensions.CommandLineUtils;
+﻿// Copyright Epic Games, Inc. All Rights Reserved.
+
+using McMaster.Extensions.CommandLineUtils;
 using Serilog;
 using System;
 using System.Collections.Generic;
diff --git a/Engine/Source/Programs/SkeinCLI/SkeinCLI/Program.cs b/Engine/Source/Programs/SkeinCLI/SkeinCLI/Program.cs
--- a/Engine/Source/Programs/SkeinCLI/SkeinCLI/Program.cs
+++ b/Engine/Source/Programs/SkeinCLI/SkeinCLI/Program.cs
@@ -1,4 +1,6 @@
-﻿using McMaster.Extensions.CommandLineUtils;
+﻿// Copyright Epic Games, Inc. All Rights Reserved.
+
+using McMaster.Extensions.CommandLineUtils;
 using Microsoft.Extensions.Configuration;
 using Microsoft.Extensions.Configuration.Json;
 using Serilog;
diff --git a/Engine/Source/Programs/SkeinCLI/SkeinCLI/ProgressBar.cs b/Engine/Source/Programs/SkeinCLI/SkeinCLI/ProgressBar.cs
deleted file mode 100644
diff --git a/Engine/Source/Programs/SkeinCLI/SkeinCLI/ProjectsCmd.cs b/Engine/Source/Programs/SkeinCLI/SkeinCLI/ProjectsCmd.cs
--- a/Engine/Source/Programs/SkeinCLI/SkeinCLI/ProjectsCmd.cs
+++ b/Engine/Source/Programs/SkeinCLI/SkeinCLI/ProjectsCmd.cs
@@ -1,4 +1,6 @@
-﻿using McMaster.Extensions.CommandLineUtils;
+﻿// Copyright Epic Games, Inc. All Rights Reserved.
+
+using McMaster.Extensions.CommandLineUtils;
 using Serilog;
 using System;
 using System.Collections.Generic;
diff --git a/Engine/Source/Programs/SkeinCLI/SkeinCLI/SkeinClient.cs b/Engine/Source/Programs/SkeinCLI/SkeinCLI/SkeinClient.cs
--- a/Engine/Source/Programs/SkeinCLI/SkeinCLI/SkeinClient.cs
+++ b/Engine/Source/Programs/SkeinCLI/SkeinCLI/SkeinClient.cs
@@ -1,6 +1,7 @@
-﻿using System.Net.Http;
-using System.Threading.Tasks;
+﻿// Copyright Epic Games, Inc. All Rights Reserved.
 
+using System.Net.Http;
+using System.Threading.Tasks;
 public class SkeinClient
 {
     private readonly HttpClient _httpClient;
diff --git a/Engine/Source/Programs/SkeinCLI/SkeinCLI/SkeinCmd.cs b/Engine/Source/Programs/SkeinCLI/SkeinCLI/SkeinCmd.cs
--- a/Engine/Source/Programs/SkeinCLI/SkeinCLI/SkeinCmd.cs
+++ b/Engine/Source/Programs/SkeinCLI/SkeinCLI/SkeinCmd.cs
@@ -1,4 +1,6 @@
-﻿using McMaster.Extensions.CommandLineUtils;
+﻿// Copyright Epic Games, Inc. All Rights Reserved.
+
+using McMaster.Extensions.CommandLineUtils;
 using System.Collections.Generic;
 using System.Reflection;
 using System.Threading.Tasks;
diff --git a/Engine/Source/Programs/SkeinCLI/SkeinCLI/SkeinCmdBase.cs b/Engine/Source/Programs/SkeinCLI/SkeinCLI/SkeinCmdBase.cs
--- a/Engine/Source/Programs/SkeinCLI/SkeinCLI/SkeinCmdBase.cs
+++ b/Engine/Source/Programs/SkeinCLI/SkeinCLI/SkeinCmdBase.cs
@@ -1,4 +1,6 @@
-﻿using McMaster.Extensions.CommandLineUtils;
+﻿// Copyright Epic Games, Inc. All Rights Reserved.
+
+using McMaster.Extensions.CommandLineUtils;
 using Serilog;
 using System;
 using System.Collections.Generic;
