Skein CLI
* Ensured that the CLI is non-reentrant by using a named Mutex
* Added -j|--json option to selectively output in JSON format instead of plain text. Added Output() and OutputError() utility functions to SkeinCmdBase.cs to simplify output formatting based on this option
* Added support for -d|--debug option to enable logging to console. By default the CLI now only logs to file

diff --git a/Engine/Source/Programs/SkeinCLI/SkeinCLI/Program.cs b/Engine/Source/Programs/SkeinCLI/SkeinCLI/Program.cs
--- a/Engine/Source/Programs/SkeinCLI/SkeinCLI/Program.cs
+++ b/Engine/Source/Programs/SkeinCLI/SkeinCLI/Program.cs
@@ -2,50 +2,66 @@
 
 using McMaster.Extensions.CommandLineUtils;
 using Microsoft.Extensions.Configuration;
-using Microsoft.Extensions.Configuration.Json;
 using Serilog;
 using System;
 using System.IO;
+using System.Threading;
 using System.Threading.Tasks;
+using System.Linq;
+using static SkeinCLI.SkeinCmdBase;
 
 namespace SkeinCLI
 {
-    class Program
-    {
-        private static async Task<int> Main(string[] args)
-        {
-            var Configuration = new ConfigurationBuilder()
-                .SetBasePath(Directory.GetCurrentDirectory())
-                .AddJsonFile(AppDomain.CurrentDomain.BaseDirectory + "\\appsettings.json", optional: true, reloadOnChange: true)
-                .AddEnvironmentVariables()
-                .Build();
-
-            Log.Logger = new LoggerConfiguration()
-               .ReadFrom.Configuration(Configuration)
-               .CreateLogger();
-
-            /*using (var progress = new ProgressBar())
-            {
-                for (int i = 0; i <= 100; i++)
-                {
-                    progress.Report((double)i / 100);
-                    System.Threading.Thread.Sleep(20);
-                }
-            }*/
-
-            try
-            {
-                return await CommandLineApplication.ExecuteAsync<SkeinCmd>(args);
-            }
-            catch (Exception ex)
-            {
-                Log.ForContext<Program>().Error(ex, "Oh no!");
-                return 1;
-            }
-            finally 
-            {
-                Log.CloseAndFlush();
-            }
-        }
-    }
+	class Program
+	{
+		private static async Task<int> Main(string[] args)
+		{
+			// Using a named Mutex we can ensure that only one instance of the application is running on the machine
+			using (var mutex = new Mutex(false, "SkeinCLI"))
+			{
+				bool IsAlreadyExecuting = !mutex.WaitOne(TimeSpan.Zero);
+				if (IsAlreadyExecuting)
+				{
+					// Returning immediately if another instance of the Skein CLI is executing
+					return (int)ReturnCodes.AlreadyExecuting;
+				}
+
+				string configFile = (args.Contains("--debug") || args.Contains("-d")) ? "appsettings-debug.json" : "appsettings.json";
+
+				Environment.SetEnvironmentVariable("SKEIN_CLI_LOGPATH", GetLogsOutputFolder());
+
+				var Configuration = new ConfigurationBuilder()
+					.SetBasePath(Directory.GetCurrentDirectory())
+					.AddJsonFile(Path.Combine(AppDomain.CurrentDomain.BaseDirectory, configFile), optional: true, reloadOnChange: true)
+					.AddEnvironmentVariables()
+					.Build();
+
+				Log.Logger = new LoggerConfiguration()
+					.ReadFrom.Configuration(Configuration)
+					.CreateLogger();
+
+				try
+				{
+					return await CommandLineApplication.ExecuteAsync<SkeinCmd>(args);
+				}
+				catch (Exception ex)
+				{
+					Log.ForContext<Program>().Error(ex, "Oh no!");
+					return (int)ReturnCodes.GenericFailure;
+				}
+				finally
+				{
+					Log.CloseAndFlush();
+					mutex.ReleaseMutex();
+				}
+			}
+		}
+
+		private static string GetLogsOutputFolder()
+		{
+			// Windows: [REDACTED_USER_HOME]\AppData\Roaming\skein\logs
+			// Linux  : [REDACTED_USER_HOME]/skein/logs
+			return Path.Combine(System.Environment.GetFolderPath(Environment.SpecialFolder.ApplicationData), "skein", "logs");
+		}
+	}
 }
diff --git a/Engine/Source/Programs/SkeinCLI/SkeinCLI/ProjectsCmd.cs b/Engine/Source/Programs/SkeinCLI/SkeinCLI/ProjectsCmd.cs
--- a/Engine/Source/Programs/SkeinCLI/SkeinCLI/ProjectsCmd.cs
+++ b/Engine/Source/Programs/SkeinCLI/SkeinCLI/ProjectsCmd.cs
@@ -55,8 +55,23 @@ namespace SkeinCLI
 
         protected override Task<int> OnExecute(CommandLineApplication app)
         {
-            Log.Logger.ForContext<ProjectsListCmd>().Information("Projects listed here [...]");
-            return base.OnExecute(app);
+			/*
+			 * TODO: Remove this.
+			 * This is a test implementation to verify the behaviour of the -j flag
+			 */
+			Random rnd = new Random();
+			if (rnd.NextDouble() > .5)
+			{
+				List<string> projects = new List<string>() { "project1", "testProject", "this is a test project name" };
+				Output("Retrieved projects", projects);
+				return base.OnExecute(app);
+			}
+			else 
+			{
+				OutputError(403, "Cannot reach server");
+				base.OnExecute(app);
+				return Task.FromResult(403);
+			}
         }
     }
 
diff --git a/Engine/Source/Programs/SkeinCLI/SkeinCLI/SkeinCLI.csproj b/Engine/Source/Programs/SkeinCLI/SkeinCLI/SkeinCLI.csproj
--- a/Engine/Source/Programs/SkeinCLI/SkeinCLI/SkeinCLI.csproj
+++ b/Engine/Source/Programs/SkeinCLI/SkeinCLI/SkeinCLI.csproj
@@ -9,10 +9,14 @@
   </PropertyGroup>
 
   <ItemGroup>
+    <None Remove="appsettings-debug.json" />
     <None Remove="appsettings.json" />
   </ItemGroup>
 
   <ItemGroup>
+    <Content Include="appsettings-debug.json">
+      <CopyToOutputDirectory>Always</CopyToOutputDirectory>
+    </Content>
     <Content Include="appsettings.json">
       <CopyToOutputDirectory>Always</CopyToOutputDirectory>
     </Content>
@@ -24,6 +28,7 @@
     <PackageReference Include="Microsoft.Extensions.Configuration.EnvironmentVariables" Version="5.0.0" />
     <PackageReference Include="Microsoft.Extensions.Configuration.Json" Version="5.0.0" />
     <PackageReference Include="Microsoft.Extensions.Http" Version="5.0.0" />
+    <PackageReference Include="Newtonsoft.Json" Version="13.0.1" />
     <PackageReference Include="Serilog" Version="2.10.0" />
     <PackageReference Include="Serilog.Settings.Configuration" Version="3.2.0" />
     <PackageReference Include="Serilog.Sinks.Console" Version="4.0.0" />
diff --git a/Engine/Source/Programs/SkeinCLI/SkeinCLI/SkeinCmd.cs b/Engine/Source/Programs/SkeinCLI/SkeinCLI/SkeinCmd.cs
--- a/Engine/Source/Programs/SkeinCLI/SkeinCLI/SkeinCmd.cs
+++ b/Engine/Source/Programs/SkeinCLI/SkeinCLI/SkeinCmd.cs
@@ -19,7 +19,7 @@ namespace SkeinCLI
         {
             // this shows help even if the --help option isn't specified
             app.ShowHelp();
-            return Task.FromResult(0);
+			return base.OnExecute(app);
         }
 
         private static string GetVersion()
diff --git a/Engine/Source/Programs/SkeinCLI/SkeinCLI/SkeinCmdBase.cs b/Engine/Source/Programs/SkeinCLI/SkeinCLI/SkeinCmdBase.cs
--- a/Engine/Source/Programs/SkeinCLI/SkeinCLI/SkeinCmdBase.cs
+++ b/Engine/Source/Programs/SkeinCLI/SkeinCLI/SkeinCmdBase.cs
@@ -1,6 +1,7 @@
 ﻿// Copyright Epic Games, Inc. All Rights Reserved.
 
 using McMaster.Extensions.CommandLineUtils;
+using Newtonsoft.Json;
 using Serilog;
 using System;
 using System.Collections.Generic;
@@ -11,26 +12,100 @@ namespace SkeinCLI
     [HelpOption("--help")]
     abstract class SkeinCmdBase
     {
-        public abstract List<string> CreateArgs();
+		public enum ReturnCodes
+		{
+			AlreadyExecuting = -1,
+			Success = 0,
+			GenericFailure = 1
+		}
+
+		public class SkeinOutput
+		{
+			[JsonProperty("code")]
+			public int Code { get; set; }
+
+			[JsonProperty("ok")]
+			public bool OK { get; set; }
+
+			[JsonProperty("message")]
+			public string Message{ get; set; }
+
+			[JsonProperty("data")]
+			public object Data { get; set; }
+		}
+
+		public abstract List<string> CreateArgs();
         
         [Option("-q|--quiet", CommandOptionType.NoValue, Description = "Execute in quiet mode")]
         public bool Quiet { get; set; }
 
-        protected virtual Task<int> OnExecute(CommandLineApplication app)
+		[Option("-j|--json", CommandOptionType.NoValue, Description = "Output result in JSON format")]
+		public bool JsonOutput { get; set; }
+
+		[Option("-d|--debug", CommandOptionType.NoValue, Description = "Run in Debug (verbose) mode")]
+		public bool Debug { get; set; }
+
+		protected virtual Task<int> OnExecute(CommandLineApplication app)
         {
             var args = CreateArgs();
             if (Quiet)
             {
                 args.Add("-q");
-            }
+			}
+
+			if (JsonOutput)
+			{
+				args.Add("-j");
+			}
+
+			if (Debug)
+			{
+				args.Add("-d");
+			}
 
             Log.Logger.ForContext<SkeinCmdBase>().Debug("Executing: skein " + ArgumentEscaper.EscapeAndConcatenate(args));
-            return Task.FromResult(0);
+            return Task.FromResult((int)ReturnCodes.Success);
         }
 
         protected void OnException(Exception ex)
         {
             Log.Logger.ForContext<SkeinCmdBase>().Error(ex, "Exception");
         }
-    }
+
+		protected void OutputError(int errorCode, string errorMessage, object data = null)
+		{
+			if (JsonOutput)
+			{
+				SkeinOutput output = new SkeinOutput() { Code = errorCode, OK = false, Message = errorMessage, Data = data };
+				string json = JsonConvert.SerializeObject(output);
+				Console.WriteLine(json);
+			}
+			else
+			{
+				Console.WriteLine($"Error {errorCode}: {errorMessage}");
+				if (data != null)
+				{
+					Console.WriteLine($"Data: {JsonConvert.SerializeObject(data)}"); 
+				}			
+			}
+		}
+
+		protected void Output(string message, object data = null)
+		{
+			if (JsonOutput)
+			{
+				SkeinOutput output = new SkeinOutput() { Code = 200, OK = true, Message = message, Data = data };
+				string json = JsonConvert.SerializeObject(output);
+				Console.WriteLine(json);
+			}
+			else
+			{
+				Console.WriteLine($"OK: {message}");
+				if (data != null)
+				{
+					Console.WriteLine($"Data: {JsonConvert.SerializeObject(data)}");
+				}
+			}
+		}
+	}
 }
\ No newline at end of file
diff --git a/Engine/Source/Programs/SkeinCLI/SkeinCLI/appsettings-debug.json b/Engine/Source/Programs/SkeinCLI/SkeinCLI/appsettings-debug.json
new file mode 100644
--- /dev/null
+++ b/Engine/Source/Programs/SkeinCLI/SkeinCLI/appsettings-debug.json
@@ -0,0 +1,22 @@
+{
+  "Serilog": {
+    "MinimumLevel": "Debug",
+    "Enrich": [ "FromLogContext", "WithExceptionDetails" ],
+    "WriteTo": [
+      {
+        "Name": "Console",
+        "Args": {
+          "outputTemplate": "{Timestamp:HH:mm:ss.fff zzz} [{Level:w3}] {Message:lj}{NewLine}{Exception}"
+        }
+      },
+      {
+        "Name": "File",
+        "Args": {
+          "path": "Logs/skein.log",
+          "rollingInterval": "Day",
+          "outputTemplate": "{Timestamp:HH:mm:ss.fff zzz} [{Level:w3}] {Message:lj}{NewLine}{Exception}"
+        }
+      }
+    ]
+  }
+}
\ No newline at end of file
diff --git a/Engine/Source/Programs/SkeinCLI/SkeinCLI/appsettings.json b/Engine/Source/Programs/SkeinCLI/SkeinCLI/appsettings.json
--- a/Engine/Source/Programs/SkeinCLI/SkeinCLI/appsettings.json
+++ b/Engine/Source/Programs/SkeinCLI/SkeinCLI/appsettings.json
@@ -3,19 +3,13 @@
     "MinimumLevel": "Debug",
     "Enrich": [ "FromLogContext", "WithExceptionDetails" ],
     "WriteTo": [
-      {
-        "Name": "Console",
-        "Args": {
-          "outputTemplate": "{Timestamp:HH:mm:ss.fff zzz} [{Level:w3}] {Message:lj}{NewLine}{Exception}"
-        }
-      },
       {
         "Name": "File",
-        "Args": {
-          "path": "Logs/skein.log",
-          "rollingInterval": "Day",
-          "outputTemplate": "{Timestamp:HH:mm:ss.fff zzz} [{Level:w3}] {Message:lj}{NewLine}{Exception}"
-        }
+		"Args": {
+			"path": "%SKEIN_CLI_LOGPATH%/skein.log",
+			"rollingInterval": "Day",
+			"outputTemplate": "{Timestamp:HH:mm:ss.fff zzz} [{Level:w3}] {Message:lj}{NewLine}{Exception}"
+		}
       }
     ]
   }
