SkeinCLI: Implement the CLI Authentication workflow.
The commands 'auth login' and 'auth logout' should now work.

#rnx

diff --git a/Engine/Source/Programs/SkeinCLI/SkeinCLI/AuthCmd.cs b/Engine/Source/Programs/SkeinCLI/SkeinCLI/AuthCmd.cs
--- a/Engine/Source/Programs/SkeinCLI/SkeinCLI/AuthCmd.cs
+++ b/Engine/Source/Programs/SkeinCLI/SkeinCLI/AuthCmd.cs
@@ -36,13 +36,7 @@ namespace SkeinCLI
     {
         private AuthCmd Parent { get; set; }
 
-        [Option("-u|--username", CommandOptionType.SingleValue, Description = "User Name")]
-        public string UserName { get; set; }
-
-        [Option("-p|--password", CommandOptionType.SingleValue, Description = "Password")]
-        public string Password { get; set; }
-
-        public override List<string> CreateArgs()
+		public override List<string> CreateArgs()
         {
             var args = Parent.CreateArgs();
             args.Add("login");
@@ -52,34 +46,28 @@ namespace SkeinCLI
 
         protected override Task<int> OnExecute(CommandLineApplication app)
         {
-            if (string.IsNullOrEmpty(UserName) && !Quiet)
-            {
-                UserName = Prompt.GetString("User Name:",
-                    promptColor: ConsoleColor.Black,
-                    promptBgColor: ConsoleColor.White);
-            }
-
-            if (string.IsNullOrEmpty(Password) && !Quiet)
-            {
-                Password = Prompt.GetPassword("Password:",
-                    promptColor: ConsoleColor.Black,
-                    promptBgColor: ConsoleColor.White);
-            }
-
-            if (!string.IsNullOrEmpty(UserName) && !string.IsNullOrEmpty(Password))
-            {
-                Log.Logger.ForContext<ProjectsSnapshotsGetCmd>().Information($"User Name: {UserName}, Password: {Password}");
-                Log.Logger.ForContext<ProjectsListCmd>().Information("Login process here [...]");
-            }
-            else 
-            {
-                Log.Logger.ForContext<ProjectsListCmd>().Information("Login cancelled");
-            }
-
-            
-            return base.OnExecute(app);
-        }
-    }
+			if (AuthUtils.IsLoggedIn())
+			{
+				Log.Logger.ForContext<AuthLoginCmd>().Information("Login succeeded (token still valid).");
+				return Task.FromResult(0);
+			}
+
+			if (AuthUtils.Refresh())
+			{
+				Log.Logger.ForContext<AuthLoginCmd>().Information("Login succeeded (token refreshed).");
+				return Task.FromResult(0);
+			}
+
+			if (AuthUtils.Login())
+			{
+				Log.Logger.ForContext<AuthLoginCmd>().Information("Login succeeded.");
+				return Task.FromResult(0);
+			}
+			
+			Log.Logger.ForContext<AuthLoginCmd>().Error("Login failed.");
+			return Task.FromResult(1);
+		}
+	}
 
     [Command(Name = "logout", Description = "Skein logout")]
     class AuthLogoutCmd : SkeinCmdBase
@@ -96,8 +84,14 @@ namespace SkeinCLI
 
         protected override Task<int> OnExecute(CommandLineApplication app)
         {
-            Log.Logger.ForContext<ProjectsListCmd>().Information("Logout process here [...]");
-            return base.OnExecute(app);
-        }
+			if (AuthUtils.Logout())
+			{
+				Log.Logger.ForContext<AuthLoginCmd>().Information("Logout succeeded.");
+				return Task.FromResult(0);
+			}
+			
+			Log.Logger.ForContext<AuthLoginCmd>().Error("Logout failed.");
+			return Task.FromResult(1);
+		}
     }
 }
\ No newline at end of file
diff --git a/Engine/Source/Programs/SkeinCLI/SkeinCLI/AuthUtils.cs b/Engine/Source/Programs/SkeinCLI/SkeinCLI/AuthUtils.cs
new file mode 100644
--- /dev/null
+++ b/Engine/Source/Programs/SkeinCLI/SkeinCLI/AuthUtils.cs
@@ -0,0 +1,682 @@
+// Copyright Epic Games, Inc. All Rights Reserved.
+
+using Newtonsoft.Json;
+using Newtonsoft.Json.Linq;
+using Serilog;
+using System;
+using System.Diagnostics;
+using System.IO;
+using System.Net;
+using System.Net.Sockets;
+using System.Runtime.InteropServices;
+using System.Security.Cryptography;
+using System.Text;
+using System.Threading.Tasks;
+
+namespace SkeinCLI
+{
+	/// <urls>
+	/// - [REDACTED_PROFILE_URL]/oauth-apps-for-windows/tree/master/OAuthConsoleApp
+	/// - https://dev.epicgames.com/docs/services/en-US/WebAPIRef/AuthWebAPI/index.html
+	/// </urls>
+
+	public static class AuthUtils
+	{
+		private const string URIForRequestCode = "https://www.epicgames.com/id/authorize";
+		private const string URIForRequestToken = "https://api.epicgames.dev/epic/oauth/v1/token";
+		private const string URIForValidateToken = "https://api.epicgames.dev/epic/oauth/v1/tokenInfo";
+		private const string URIForBrowserRedirect = "https://www.epicgames.com/account/connections?lang=en&productName=epicgames";
+		private const string URIForOAuthRedirect = "http://127.0.0.1";
+		private const string ClientId = "xyza7891KnciMbgQ7gGk4HnaHphs0IPX";
+		private const string ClientSecret = "utRF4sYZFa7ZNXiX9QItLDlllWGx2lq/c9qFcPaXycw";
+		private const string DeploymentId = "b95fab5dc84d4e2ea95588b538172289";
+		private const string FileNameAccessToken = "access.token";
+		private const string FileNameRefreshToken = "refresh.token";
+		private const int PortMin = 6970;
+		private const int PortMax = 6975;
+
+		/// <summary>
+		/// Checks if the given port is open on the local machine.
+		/// </summary>
+		/// <param name="port"></param>
+		private static bool IsPortOpen(int port)
+		{
+			try
+			{
+				IPAddress ipAddress = Dns.GetHostEntry("localhost").AddressList[0];
+				TcpListener tcpListener = new TcpListener(ipAddress, port);
+				tcpListener.Start();
+				tcpListener.Stop();
+				return true;
+			}
+			catch
+			{
+				return false;
+			}
+		}
+
+		/// <summary>
+		/// Checks if the user is logged in by validating if a stored token is still valid.
+		/// Do not call this on a time sensitive thread as it potentially does a web query!
+		/// </summary>
+		/// <returns></returns>
+		public static bool IsLoggedIn()
+		{
+			string dummy = string.Empty;
+			return IsLoggedIn(ref dummy);
+		}
+
+		/// <summary>
+		/// Checks if the user is logged in by validating if a stored token is still valid.
+		/// Do not call this on a time sensitive thread as it potentially does a web query!
+		/// Also returns the token for subsequent calls to the Skein API.
+		/// </summary>
+		/// <param name="token"></param>
+		/// <returns></returns>
+		public static bool IsLoggedIn(ref string token)
+		{
+			string accessToken = null;
+			string refreshToken = null;
+
+			// grab the tokens that are stored locally
+			try
+			{
+				if (!FetchLocalTokens(ref accessToken, ref refreshToken))
+				{
+					return false;
+				}
+			}
+			catch (Exception ex)
+			{
+				Log.Logger.Error(ex, "Exception while executing AuthUtils::IsLoggedIn");
+				return false;
+			}
+
+			// check if the access token is still valid
+			try
+			{
+				Task.Run(async () => { await DoValidateTokenAsync(accessToken, refreshToken); }).Wait();
+			}
+			catch
+			{
+				// it's expected that this throws if the access token isn't valid so not logging this
+				return false;
+			}
+
+			token = accessToken;
+			return true;
+		}
+
+		/// <summary>
+		/// Performs a Login cycle.
+		/// </summary>
+		public static bool Login()
+		{
+			// try to obtain a new token
+			try
+			{
+				Task.Run(async () => { await DoRequestCodeAsync(ClientId, ClientSecret, DeploymentId); }).Wait();
+			}
+			catch (Exception ex)
+			{
+				Log.Logger.Error(ex, "Exception while executing AuthUtils::Login");
+				return false;
+			}
+
+			return IsLoggedIn();
+		}
+
+		/// <summary>
+		/// Performs a Logout by clearing the local tokens.
+		/// Not an actual logout in the browser.
+		/// </summary>
+		/// <returns></returns>
+		public static bool Logout()
+		{
+			try
+			{
+				ClearLocalTokens();
+			}
+			catch (Exception ex)
+			{
+				Log.Logger.Error(ex, "Exception while executing AuthUtils::Logout");
+				return false;
+			}
+
+			return !IsLoggedIn();
+		}
+
+		/// <summary>
+		/// Refresh the access token.
+		/// </summary>
+		/// <returns></returns>
+		public static bool Refresh()
+		{
+			string accessToken = null;
+			string refreshToken = null;
+			// grab the tokens that are stored locally
+			try
+			{
+				if (!FetchLocalTokens(ref accessToken, ref refreshToken))
+				{
+					return false;
+				}
+			}
+			catch (Exception ex)
+			{
+				Log.Logger.Error(ex, "Exception while executing AuthUtils::Refresh");
+				return false;
+			}
+
+			// try to utilize the refresh token to renew our session
+			// this is nice as it doesn't popup the browser of the user
+			try
+			{
+				Task.Run(async () => { await DoRefreshTokenAsync(ClientId, ClientSecret, DeploymentId, refreshToken); }).Wait();
+			}
+			catch
+			{
+				// it's expected that this throws if the refresh token isn't valid so not logging this
+				return false;
+			}
+
+			return IsLoggedIn();
+		}
+
+		/// <summary>
+		/// Reads the tokens stored on the local machine.
+		/// </summary>
+		/// <param name="accessToken"></param>
+		/// <param name="refreshToken"></param>
+		/// <returns>true if a token was found</returns>
+		private static bool FetchLocalTokens(ref string accessToken, ref string refreshToken)
+		{
+			string accessTokenPath = Path.Combine(GetTokenFolder(), FileNameAccessToken);
+			if (File.Exists(accessTokenPath))
+			{
+				FileStream stream = new FileStream(accessTokenPath, FileMode.Open);
+				using (StreamReader reader = new StreamReader(stream, Encoding.ASCII))
+				{
+					accessToken = reader.ReadLine();
+				}
+			}
+
+			string refreshTokenPath = Path.Combine(GetTokenFolder(), FileNameRefreshToken);
+			if (File.Exists(refreshTokenPath))
+			{
+				FileStream stream = new FileStream(refreshTokenPath, FileMode.Open);
+				using (StreamReader reader = new StreamReader(stream, Encoding.ASCII))
+				{
+					refreshToken = reader.ReadLine();
+				}
+			}
+
+			return accessToken != null && accessToken.Length != 0 && refreshToken != null && refreshToken.Length != 0;
+		}
+
+		/// <summary>
+		/// Clears the tokens on the local machine.
+		/// </summary>
+		private static void ClearLocalTokens()
+		{
+			StoreLocalTokens("", "");
+		}
+
+		/// <summary>
+		/// Stores the tokens on the local machine for future use.
+		/// </summary>
+		/// <param name="accessToken"></param>
+		/// <param name="refreshToken"></param>
+		/// <returns>true if the token was stored</returns>
+		private static void StoreLocalTokens(string accessToken, string refreshToken)
+		{
+			string tokenFolder = GetTokenFolder();
+
+			// create the directory if it doesn't already exist
+			if (!Directory.Exists(tokenFolder))
+			{
+				Directory.CreateDirectory(tokenFolder);
+			}
+
+			// attempt to write the tokens to their files
+			FileStream stream = null;
+				
+			stream = new FileStream(Path.Combine(tokenFolder, FileNameAccessToken), FileMode.Create);
+			using (StreamWriter writer = new StreamWriter(stream, Encoding.ASCII))
+			{
+				writer.WriteLine(accessToken);
+			}
+
+			stream = new FileStream(Path.Combine(tokenFolder, FileNameRefreshToken), FileMode.Create);
+			using (StreamWriter writer = new StreamWriter(stream, Encoding.ASCII))
+			{
+				writer.WriteLine(refreshToken);
+			}
+		}
+
+		/// <summary>
+		/// Returns the folder where the tokens will be stored.
+		/// </summary>
+		private static string GetTokenFolder()
+		{
+			// Windows: [REDACTED_USER_HOME]\AppData\Roaming\skein
+			// Linux  : [REDACTED_USER_HOME]/skein
+
+			return Path.Combine(System.Environment.GetFolderPath(Environment.SpecialFolder.ApplicationData), "skein");
+		}
+
+		/// <summary>
+		/// Returns the path to use as the RedirectURI for the OAuth sequence.
+		/// Requires an open port on the local machine.
+		/// The URL needs to be white-listed in the EpicGames app portal!
+		/// </summary>
+		private static string GetRedirectURI()
+		{
+			for (int port = PortMin; port < PortMax; port++)
+			{
+				if (IsPortOpen(port))
+				{
+					return string.Format("{0}:{1}/", URIForOAuthRedirect, port);
+				}
+			}
+
+			return string.Format("{0}:{1}/", URIForOAuthRedirect, PortMin);
+		}
+
+		/// <summary>
+		/// Returns URI-safe data with a given input length.
+		/// </summary>
+		/// <param name="length">Input length (nb. output will be longer)</param>
+		private static string GetRandomDataBase64url(uint length)
+		{
+			RNGCryptoServiceProvider rng = new RNGCryptoServiceProvider();
+			byte[] bytes = new byte[length];
+			rng.GetBytes(bytes);
+			return DoBase64UrlEncodeNoPadding(bytes);
+		}
+
+		/// <summary>
+		/// Base64url no-padding encodes the given input buffer.
+		/// </summary>
+		/// <param name="buffer"></param>
+		private static string DoBase64UrlEncodeNoPadding(byte[] buffer)
+		{
+			string base64 = Convert.ToBase64String(buffer);
+
+			// converts base64 to base64url
+			base64 = base64.Replace("+", "-");
+			base64 = base64.Replace("/", "_");
+			// strips padding
+			base64 = base64.Replace("=", "");
+
+			return base64;
+		}
+
+		/// <summary>
+		/// OAuth: request code.
+		/// </summary>
+		/// <param name="clientId"></param>
+		/// <param name="clientSecret"></param>
+		/// <param name="deploymentId"></param>
+		/// <returns></returns>
+		private static async Task DoRequestCodeAsync(string clientId, string clientSecret, string deploymentId)
+		{
+			// generates state value
+			string state = GetRandomDataBase64url(32);
+			string redirect = GetRedirectURI();
+
+			// creates an HttpListener to listen for requests on Skein redirect URI
+			var http = new HttpListener();
+			http.Prefixes.Add(redirect);
+			http.Start();
+
+			// creates the OAuth 2.0 authorization request
+			string uriForAuthorizationRequest = string.Format("{0}?response_type=code&scope=basic_profile&redirect_uri={1}&client_id={2}&state={3}",
+				URIForRequestCode,
+				Uri.EscapeDataString(redirect),
+				clientId,
+				state);
+
+			// opens request in the default browser
+			Process process = null;
+			try
+			{
+				process = Process.Start(uriForAuthorizationRequest);
+			}
+			catch
+			{
+				// hack because of this: [REDACTED_PROFILE_URL]/corefx/issues/10361
+				if (RuntimeInformation.IsOSPlatform(OSPlatform.Windows))
+				{
+					uriForAuthorizationRequest = uriForAuthorizationRequest.Replace("&", "^&");
+					process = Process.Start(new ProcessStartInfo("cmd", $"/c start {uriForAuthorizationRequest}"));
+				}
+				else if (RuntimeInformation.IsOSPlatform(OSPlatform.Linux))
+				{
+					process = Process.Start("xdg-open", uriForAuthorizationRequest);
+				}
+				else if (RuntimeInformation.IsOSPlatform(OSPlatform.OSX))
+				{
+					process = Process.Start("open", uriForAuthorizationRequest);
+				}
+				else
+				{
+					throw new PlatformNotSupportedException();
+				}
+			}
+
+			// waits for the OAuth authorization response
+			var context = await http.GetContextAsync();
+
+			// sends an HTTP redirect response to the browser
+			// accompany it with an empty string to ensure that the client finished reading the response
+			// if we call http.Stop() too soon the user be presented a browser error instead
+			byte[] responseBuffer = Encoding.UTF8.GetBytes("");
+
+			var response = context.Response;
+			response.ContentLength64 = responseBuffer.Length;
+			response.Redirect(URIForBrowserRedirect);
+
+			var responseOutput = response.OutputStream;
+			responseOutput.Write(responseBuffer, 0, responseBuffer.Length);
+			responseOutput.Close();
+
+			// stop listening
+			http.Stop();
+
+			// checks for errors
+			string error = context.Request.QueryString.Get("error");
+			if (error is object)
+			{
+				throw new Exception(error.ToString());
+			}
+			if (context.Request.QueryString.Get("code") is null || context.Request.QueryString.Get("state") is null)
+			{
+				throw new Exception(context.Request.QueryString.ToString());
+			}
+
+			// extract the incoming data
+			var incomingCode = context.Request.QueryString.Get("code");
+			var incomingState = context.Request.QueryString.Get("state");
+
+			// compares the received state to the expected value, to ensure that this app made the request which resulted in authorization
+			if (incomingState != state)
+			{
+				throw new Exception(string.Format("{0} != {1}", incomingState, state));
+			}
+
+			// starts the code exchange
+			await DoRequestTokenAsync(ClientId, ClientSecret, DeploymentId, incomingCode);
+		}
+
+		/// <summary>
+		/// OAuth: request token.
+		/// </summary>
+		/// <param name="clientId"></param>
+		/// <param name="clientSecret"></param>
+		/// <param name="deploymentId"></param>
+		/// <param name="code"></param>
+		/// <returns></returns>
+		private static async Task DoRequestTokenAsync(string clientId, string clientSecret, string deploymentId, string code)
+		{
+			string tokenRequestBody = string.Format("code={0}&client_id={1}&client_secret={2}&deployment_id={3}&scope=basic_profile&grant_type=authorization_code",
+				code,
+				clientId,
+				clientSecret,
+				deploymentId
+				);
+
+			// build the data buffers
+			byte[] credentialBufferBytes = Encoding.ASCII.GetBytes(clientId + ":" + clientSecret);
+			byte[] tokenRequestBodyBytes = Encoding.ASCII.GetBytes(tokenRequestBody);
+
+			// send the request
+			WebRequest tokenRequest = WebRequest.Create(URIForRequestToken);
+			tokenRequest.Method = "POST";
+			tokenRequest.ContentType = "application/x-www-form-urlencoded";
+			tokenRequest.Headers["Authorization"] = "Basic " + Convert.ToBase64String(credentialBufferBytes);
+			tokenRequest.ContentLength = tokenRequestBodyBytes.Length;
+			using (Stream requestStream = tokenRequest.GetRequestStream())
+			{
+				await requestStream.WriteAsync(tokenRequestBodyBytes, 0, tokenRequestBodyBytes.Length);
+			}
+
+			try
+			{
+				// example response:
+				//
+				// {
+				//   "scope":"basic_profile",
+				//   "token_type":"bearer",
+				//   "access_token":"eyJ0IjoiZXBpY19pZCIsImFsZyI6IlJTMjU2Iiwia2lkIjoibldVQzlxSFVldWRHcnBXb3FvVXVHZkFIYmVWM2NsRnlsdFRYMzhFbXJKSSJ9.eyJhdWQiOiJ4eXphNzg5MWxoeE1WWUdDT043TGduS1paOEhRR0Q1SCIsInN1YiI6Ijk2MjZmNDQxMDU1MzQ5Y2U4Y2I3ZDdkNWE0ODNlYWEyIiwidCI6ImVwaWNfaWQiLCJzY29wZSI6ImJhc2ljX3Byb2ZpbGUgZnJpZW5kc19saXN0IHByZXNlbmNlIiwiYXBwaWQiOiJmZ2hpNDU2N08wM0hST3hFandibjdrZ1hwQmhuaFd3diIsImlzcyI6Imh0dHBzOlwvXC9hcGkuZXBpY2dhbWVzLmRldlwvZXBpY1wvb2F1dGhcL3YxIiwiZG4iOiJLcm5icnkiLCJleHAiOjE1ODgyODYwODMsImlhdCI6MTU4ODI3ODg4Mywibm9uY2UiOiJuLUI1cGNsSXZaSkJaQU1KTDVsNkdvUnJDTzNiRT0iLCJqdGkiOiI2NGMzMGQwMjk4YTM0MzdjOGE3NGU1OTAxYzM0ODZiNSJ9.MZRoCRpjIb--dD7hxoo2GvjSPhUSNpOq1FhtShTBmzMJ1qlHFPzNaUiAEETAc3mabGPKyOxUP6Q1FBadr_P_UtbtB7kf34hN2VTv5czW6WOx1HdpjwUQZuxFyDc_aix7FCS0Egu4rZlC65b-B0FUVlial_s_FrH8ou5L_d-4I0KVpIwtv-b_M6EQ9jtLdQRfMaP6aV0rIerrbqFZ617Pe7XT4IO9jZFwM8F5aDTeDHkkOO41wyVibrm38799lP4B65RIv9CwbAL-TVmV1L5gFYITaZhi5ShfZzTvxAk-3Dxwp8c5JvcO68zpbya5gFSAfhsd7vt9YLU0gQR2uXq3Vw",
+				//   "refresh_token":"eyJ0IjoiZXBpY19pZCIsImFsZyI6IlJTMjU2Iiwia2lkIjoibldVQzlxSFVldWRHcnBXb3FvVXVHZkFIYmVWM2NsRnlsdFRYMzhFbXJKSSJ9.eyJhdWQiOiJ4eXphNzg5MWxoeE1WWUdDT043TGduS1paOEhRR0Q1SCIsInN1YiI6Ijk2MjZmNDQxMDU1MzQ5Y2U4Y2I3ZDdkNWE0ODNlYWEyIiwidCI6ImVwaWNfaWQiLCJhcHBpZCI6ImZnaGk0NTY3TzAzSFJPeEVqd2JuN2tnWHBCaG5oV3d2Iiwic2NvcGUiOiJiYXNpY19wcm9maWxlIGZyaWVuZHNfbGlzdCBwcmVzZW5jZSIsImlzcyI6Imh0dHBzOlwvXC9hcGkuZXBpY2dhbWVzLmRldlwvZXBpY1wvb2F1dGhcL3YxIiwiZG4iOiJLcm5icnkiLCJleHAiOjE1ODgzMDc2ODMsImlhdCI6MTU4ODI3ODg4MywianRpIjoiYzczYjA2NmUyZDU4NGVkNTk0NjZiOThiNzI3NzJiMjAifQ.O-eVa46NimubKwxe9SwlHxciivu0XWe1-DSL74mMiA_PpPoW0yKL9DfmsLxiPCwsRB5_hQTc6_FM7G1FyfKtX_VVAp90MZPkhCbAbfKmTpQVcL0Ya6kve4KMG8KxeLVfLLhubCbJTYlnDNVHobbpvpQtHd8Ys321ZNDJj05l_tnZzdgus-xmCO6orX4UP4wDd1jAOXXeqRT47OXuLCgSE0q6Osfh-ENPwh6ph1i7ld759xPV0oNcQb8XiPxnT6_FUmFugzG1YS1z9bTnVWmbP2RmYluue5VQm5EKGJZ91Alve8s2eNEtDfUqaBLZ45pqGkc1KjbYTtP0a_1ue2BpkQ",
+				//   "expires_in":7200,
+				//   "expires_at":"2021-10-05T09:19:26.728Z",
+				//   "refresh_expires_in":28800,
+				//   "refresh_expires_at":"2021-10-05T15:19:26.728Z",
+				//   "account_id":"9626f441055349ce8cb7d7d5a483eaa2",
+				//   "client_id":"xyza7891lhxMVYGCON7LgnKZZ8HQGD5H",
+				//   "application_id":"fghi4567O03HROxEjwbn7kgXpBhnhWwv"
+				//   "merged_accounts":[]
+				// }
+
+				// gets the response
+				WebResponse tokenResponse = await tokenRequest.GetResponseAsync();
+				using (StreamReader reader = new StreamReader(tokenResponse.GetResponseStream()))
+				{
+					// reads response body
+					string responseText = await reader.ReadToEndAsync();
+
+					// converts to json object
+					var jsonObject = JObject.Parse(responseText);
+					var jsonAcessToken = jsonObject.SelectToken("access_token");
+					var jsonRefreshToken = jsonObject.SelectToken("refresh_token");
+
+					// extract the tokens
+					string receivedAccessToken = jsonAcessToken != null ? jsonAcessToken.ToString() : null;
+					if (receivedAccessToken == null || receivedAccessToken.Length == 0)
+					{
+						throw new Exception(responseText);
+					}
+					string receivedRefreshToken = jsonRefreshToken != null ? jsonRefreshToken.ToString() : null;
+					if (receivedRefreshToken == null || receivedRefreshToken.Length == 0)
+					{
+						throw new Exception(responseText);
+					}
+
+					// validate it
+					await DoValidateTokenAsync(receivedAccessToken, receivedRefreshToken);
+				}
+			}
+			catch (JsonReaderException ex)
+			{
+				throw ex;
+			}
+			catch (WebException ex)
+			{
+				if (ex.Status == WebExceptionStatus.ProtocolError)
+				{
+					var response = ex.Response as HttpWebResponse;
+					if (response != null)
+					{
+						using (StreamReader reader = new StreamReader(response.GetResponseStream()))
+						{
+							string responseText = await reader.ReadToEndAsync();
+							throw new Exception(responseText);
+						}
+					}
+				}
+				throw ex;
+			}
+		}
+
+		/// <summary>
+		/// OAuth: refresh token.
+		/// </summary>
+		/// <param name="clientId"></param>
+		/// <param name="clientSecret"></param>
+		/// <param name="deploymentId"></param>
+		/// <param name="refreshToken"></param>
+		/// <returns></returns>
+		private static async Task DoRefreshTokenAsync(string clientId, string clientSecret, string deploymentId, string refreshToken)
+		{
+			string tokenRequestBody = string.Format("refresh_token={0}&client_id={1}&client_secret={2}&deployment_id={3}&scope=basic_profile&grant_type=refresh_token",
+				refreshToken,
+				clientId,
+				clientSecret,
+				deploymentId
+				);
+
+			// clear local tokens
+			ClearLocalTokens();
+
+			// build the data buffers
+			byte[] credentialBufferBytes = Encoding.ASCII.GetBytes(clientId + ":" + clientSecret);
+			byte[] tokenRequestBodyBytes = Encoding.ASCII.GetBytes(tokenRequestBody);
+
+			// send the request
+			WebRequest tokenRequest = WebRequest.Create(URIForRequestToken);
+			tokenRequest.Method = "POST";
+			tokenRequest.ContentType = "application/x-www-form-urlencoded";
+			tokenRequest.Headers["Authorization"] = "Basic " + Convert.ToBase64String(credentialBufferBytes);
+			tokenRequest.ContentLength = tokenRequestBodyBytes.Length;
+			using (Stream requestStream = tokenRequest.GetRequestStream())
+			{
+				await requestStream.WriteAsync(tokenRequestBodyBytes, 0, tokenRequestBodyBytes.Length);
+			}
+
+			try
+			{
+				// gets the response
+				WebResponse tokenResponse = await tokenRequest.GetResponseAsync();
+				using (StreamReader reader = new StreamReader(tokenResponse.GetResponseStream()))
+				{
+					// reads response body
+					string responseText = await reader.ReadToEndAsync();
+
+					// converts to json object
+					var jsonObject = JObject.Parse(responseText);
+					var jsonAcessToken = jsonObject.SelectToken("access_token");
+					var jsonRefreshToken = jsonObject.SelectToken("refresh_token");
+
+					// extract the tokens
+					string receivedAccessToken = jsonAcessToken != null ? jsonAcessToken.ToString() : null;
+					if (receivedAccessToken == null || receivedAccessToken.Length == 0)
+					{
+						throw new Exception(responseText);
+					}
+					string receivedRefreshToken = jsonRefreshToken != null ? jsonRefreshToken.ToString() : null;
+					if (receivedRefreshToken == null || receivedRefreshToken.Length == 0)
+					{
+						throw new Exception(responseText);
+					}
+
+					// validate it
+					await DoValidateTokenAsync(receivedAccessToken, receivedRefreshToken);
+				}
+			}
+			catch (JsonReaderException ex)
+			{
+				throw ex;
+			}
+			catch (WebException ex)
+			{
+				if (ex.Status == WebExceptionStatus.ProtocolError)
+				{
+					var response = ex.Response as HttpWebResponse;
+					if (response != null)
+					{
+						using (StreamReader reader = new StreamReader(response.GetResponseStream()))
+						{
+							string responseText = await reader.ReadToEndAsync();
+							throw new Exception(responseText);
+						}
+					}
+				}
+				throw ex;
+			}
+		}
+
+		/// <summary>
+		/// OAuth: validate token.
+		/// </summary>
+		/// <param name="accessToken"></param>
+		/// <param name="refreshToken"></param>
+		/// <returns></returns>
+		private static async Task DoValidateTokenAsync(string accessToken, string refreshToken)
+		{
+			string validateTokenBody = string.Format("token={0}", accessToken);
+			
+			// build the data buffer
+			byte[] validateTokenBodyBytes = Encoding.ASCII.GetBytes(validateTokenBody);
+
+			// send the request
+			WebRequest validateToken = WebRequest.Create(URIForValidateToken);
+			validateToken.Method = "POST";
+			validateToken.ContentType = "application/x-www-form-urlencoded";
+			validateToken.ContentLength = validateTokenBodyBytes.Length;
+			using (Stream requestStream = validateToken.GetRequestStream())
+			{
+				await requestStream.WriteAsync(validateTokenBodyBytes, 0, validateTokenBodyBytes.Length);
+			}
+
+			try
+			{
+				// example response:
+				//
+				// {
+				//   "active":true,
+				//   "scope":"basic_profile",
+				//   "token_type":"bearer",
+				//   "expires_in":7082,
+				//   "expires_at":"2021-10-05T09:19:26.728Z",
+				//   "account_id":"9626f441055349ce8cb7d7d5a483eaa2",
+				//   "client_id":"xyza7891lhxMVYGCON7LgnKZZ8HQGD5H",
+				//   "application_id":"fghi4567O03HROxEjwbn7kgXpBhnhWwv"
+				// }
+
+				// gets the response
+				WebResponse tokenResponse = await validateToken.GetResponseAsync();
+				using (StreamReader reader = new StreamReader(tokenResponse.GetResponseStream()))
+				{
+					// reads response body
+					string responseText = await reader.ReadToEndAsync();
+
+					// converts to json object
+					var jsonObject = JObject.Parse(responseText);
+					var jsonAccountId = jsonObject.SelectToken("account_id");
+
+					// extract the account id
+					string receivedAccountId = jsonAccountId != null ? jsonAccountId.ToString() : null;
+					if (receivedAccountId == null || receivedAccountId.Length == 0)
+					{
+						throw new Exception(responseText);
+					}
+
+					// store the tokens
+					StoreLocalTokens(accessToken, refreshToken);
+				}
+			}
+			catch (JsonReaderException ex)
+			{
+				throw ex;
+			}
+			catch (WebException ex)
+			{
+				if (ex.Status == WebExceptionStatus.ProtocolError)
+				{
+					var response = ex.Response as HttpWebResponse;
+					if (response != null)
+					{
+						using (StreamReader reader = new StreamReader(response.GetResponseStream()))
+						{
+							string responseText = await reader.ReadToEndAsync();
+							throw new Exception(responseText);
+						}
+					}
+				}
+				throw ex;
+			}
+		}
+	}
+}
\ No newline at end of file
diff --git a/Engine/Source/Programs/SkeinCLI/SkeinCLI/SkeinCLI.csproj b/Engine/Source/Programs/SkeinCLI/SkeinCLI/SkeinCLI.csproj
--- a/Engine/Source/Programs/SkeinCLI/SkeinCLI/SkeinCLI.csproj
+++ b/Engine/Source/Programs/SkeinCLI/SkeinCLI/SkeinCLI.csproj
@@ -21,6 +21,13 @@
       <CopyToOutputDirectory>Always</CopyToOutputDirectory>
     </Content>
   </ItemGroup>
+  
+  <ItemGroup>
+    <Reference Include="Newtonsoft.JSON">
+      <HintPath>..\..\..\..\Binaries\ThirdParty\Newtonsoft\Newtonsoft.Json.dll</HintPath>
+      <Private>True</Private>
+    </Reference>
+  </ItemGroup>
 
   <ItemGroup>
     <PackageReference Include="McMaster.Extensions.CommandLineUtils" Version="3.1.0" />
