SkeinCLI: Introduced SkeinUtils for global utility functions that are useful throughout the program.
- Added method SkeinUtils::GetDataFolder, which is now used by AuthUtils::GetTokenFolder and Program::GetLogsOutputFolder.

#rnx

diff --git a/Engine/Source/Programs/SkeinCLI/SkeinCLI/AuthUtils.cs b/Engine/Source/Programs/SkeinCLI/SkeinCLI/AuthUtils.cs
--- a/Engine/Source/Programs/SkeinCLI/SkeinCLI/AuthUtils.cs
+++ b/Engine/Source/Programs/SkeinCLI/SkeinCLI/AuthUtils.cs
@@ -259,10 +259,7 @@ namespace SkeinCLI
 		/// </summary>
 		private static string GetTokenFolder()
 		{
-			// Windows: [REDACTED_USER_HOME]\AppData\Roaming\skein
-			// Linux  : [REDACTED_USER_HOME]/skein
-
-			return Path.Combine(System.Environment.GetFolderPath(Environment.SpecialFolder.ApplicationData), "skein");
+			return SkeinUtils.GetDataFolder();
 		}
 
 		/// <summary>
diff --git a/Engine/Source/Programs/SkeinCLI/SkeinCLI/Program.cs b/Engine/Source/Programs/SkeinCLI/SkeinCLI/Program.cs
--- a/Engine/Source/Programs/SkeinCLI/SkeinCLI/Program.cs
+++ b/Engine/Source/Programs/SkeinCLI/SkeinCLI/Program.cs
@@ -59,9 +59,7 @@ namespace SkeinCLI
 
 		private static string GetLogsOutputFolder()
 		{
-			// Windows: [REDACTED_USER_HOME]\AppData\Roaming\skein\logs
-			// Linux  : [REDACTED_USER_HOME]/skein/logs
-			return Path.Combine(System.Environment.GetFolderPath(Environment.SpecialFolder.ApplicationData), "skein", "logs");
+			return Path.Combine(SkeinUtils.GetDataFolder(), "logs");
 		}
 	}
 }
diff --git a/Engine/Source/Programs/SkeinCLI/SkeinCLI/SkeinUtils.cs b/Engine/Source/Programs/SkeinCLI/SkeinCLI/SkeinUtils.cs
new file mode 100644
--- /dev/null
+++ b/Engine/Source/Programs/SkeinCLI/SkeinCLI/SkeinUtils.cs
@@ -0,0 +1,21 @@
+// Copyright Epic Games, Inc. All Rights Reserved.
+
+using System;
+using System.IO;
+
+namespace SkeinCLI
+{
+	public static class SkeinUtils
+	{
+		/// <summary>
+		/// Returns the root folder where all Skein data should be stored.
+		/// </summary>
+		public static string GetDataFolder()
+		{
+			// Windows: [REDACTED_USER_HOME]\AppData\Roaming\skein
+			// Linux  : [REDACTED_USER_HOME]/skein
+
+			return Path.Combine(System.Environment.GetFolderPath(Environment.SpecialFolder.ApplicationData), "skein");
+		}
+	}
+}
\ No newline at end of file
