SkeinCLI: Added framework for concurrent file transfers (up/down).
- Using HttpClient as a parameter of FileDownload / FileUpload as there should be only one instance of HttpClient per application and SkeinClient is using it as well (also supplied as a parameter).
- Using a semaphore for throttling the number of simultaneous transfers. Defaults to 4 but other values are possible. The throttling is per-instance (of a FileDownload or a FileUpload) and not per-application. In practice, it's likely the same thing as we shouldn't have more than one instance of a FileDownload / FileUpload at a time.
- Using a CancellationToken support to abort all pending transfers. We might want to use this if one of the transfers fails perhaps as we'd consider the entire command failed as well and can abort?
- Using events for transfer Started/Completed notifications.

#rnx

diff --git a/Engine/Source/Programs/SkeinCLI/SkeinCLI/FileTransfer.cs b/Engine/Source/Programs/SkeinCLI/SkeinCLI/FileTransfer.cs
new file mode 100644
--- /dev/null
+++ b/Engine/Source/Programs/SkeinCLI/SkeinCLI/FileTransfer.cs
@@ -0,0 +1,193 @@
+// Copyright Epic Games, Inc. All Rights Reserved.
+
+using Serilog;
+using System;
+using System.Collections.Generic;
+using System.IO;
+using System.Linq;
+using System.Net;
+using System.Net.Http;
+using System.Threading;
+using System.Threading.Tasks;
+
+namespace SkeinCLI
+{
+	public abstract class FileTransfer
+	{
+		public event EventHandler<TransferStartedEventArgs>   TransferStarted;
+		public event EventHandler<TransferCompletedEventArgs> TransferCompleted;
+
+		protected readonly HttpClient     _httpClient;
+		protected readonly int            _maxInFlight;
+		protected CancellationTokenSource _cancellationTokenSource;
+
+		/// <summary>
+		/// Constructor.
+		/// </summary>
+		/// <param name="httpClient"></param>
+		/// <param name="maxInFlight"></param>
+		protected FileTransfer(HttpClient httpClient, int maxInFlight)
+		{
+			_httpClient = httpClient;
+			_maxInFlight = maxInFlight;
+			_cancellationTokenSource = new CancellationTokenSource();
+		}
+
+		/// <summary>
+		/// Describes an item we want to transfer.
+		/// </summary>
+		public class Item
+		{
+			/// <param name="source">Either an Uri or a Path</param>
+			/// <param name="destination">Either an Uri or a Path</param>
+			public Item(string source, string destination)
+			{
+				Source = source;
+				Destination = destination;
+			}
+
+			public string Source { get; }
+			public string Destination { get; }
+		}
+
+		/// <summary>
+		/// Transfers the files from sources to their respective destinations.
+		/// </summary>
+		/// <param name="items"></param>
+		/// <returns></returns>
+		public async Task TransferFilesAsync(List<Item> items)
+		{
+			var tasks = new List<Task>();
+
+			var throttler = new SemaphoreSlim(_maxInFlight);
+			foreach (var item in items)
+			{
+				tasks.Add(
+					Task.Run(async () =>
+					{
+						// wait until the semaphore is acquired
+						await throttler.WaitAsync();
+
+						try
+						{
+							long fileSize = await DoFileSizeAsync(item.Source);
+
+							TransferStarted?.Invoke(this, new TransferStartedEventArgs(item.Source, item.Destination, fileSize));
+
+							await DoFileTransferAsync(item.Source, item.Destination);
+
+							TransferCompleted?.Invoke(this, new TransferCompletedEventArgs(item.Source, item.Destination, true));
+						}
+						catch (TaskCanceledException ex)
+						{
+							Log.Logger.ForContext<FileTransfer>().Debug(ex,"TaskCanceledException while executing FileTransfer::TransferFilesAsync");
+
+							TransferCompleted?.Invoke(this, new TransferCompletedEventArgs(item.Source, item.Destination, false));
+						}
+						catch (Exception ex)
+						{
+							Log.Logger.ForContext<FileTransfer>().Error(ex, "Exception while executing FileTransfer::TransferFilesAsync");
+
+							TransferCompleted?.Invoke(this, new TransferCompletedEventArgs(item.Source, item.Destination, false));
+						}
+						finally
+						{
+							// do not forget to release the semaphore
+							throttler.Release();
+						}
+					})
+				);
+			}
+
+			await Task.WhenAll(tasks);
+		}
+		
+		/// <summary>
+		/// Cancels all pending file transfers.
+		/// </summary>
+		public void Cancel()
+		{
+			_cancellationTokenSource.Cancel(false);
+		}
+
+		/// <summary>
+		/// Performs a file transfer from the source to the destination.
+		/// </summary>
+		/// <param name="source"></param>
+		/// <param name="destination"></param>
+		/// <returns></returns>
+		protected abstract Task DoFileTransferAsync(string source, string destination);
+
+		/// <summary>
+		/// Determines the number of bytes that need to be transferred.
+		/// </summary>
+		/// <param name="source"></param>
+		/// <returns></returns>
+		protected abstract Task<long> DoFileSizeAsync(string source);
+	}
+
+	public class FileDownload : FileTransfer
+	{
+		const int MaxDownloadsInFlight = 4;
+
+		public FileDownload(HttpClient httpClient)
+			: base(httpClient, MaxDownloadsInFlight) 
+		{
+		}
+
+		public FileDownload(HttpClient httpClient, int maxDownloadsInFlight)
+			: base(httpClient, maxDownloadsInFlight)
+		{
+		}
+
+		protected override async Task DoFileTransferAsync(string source, string destination)
+		{
+			using (var response = await _httpClient.GetStreamAsync(source))
+			{
+				using (var fileStream = new FileStream(destination, FileMode.Create))
+				{
+					await response.CopyToAsync(fileStream, _cancellationTokenSource.Token);
+				}
+			}
+		}
+
+		protected override async Task<long> DoFileSizeAsync(string source)
+		{
+			using (var response = await _httpClient.GetAsync(source, HttpCompletionOption.ResponseHeadersRead, _cancellationTokenSource.Token))
+			{
+				return response.Content.Headers.ContentLength.GetValueOrDefault();
+			}
+		}
+	}
+
+	public class FileUpload : FileTransfer
+	{
+		const int MaxUploadsInFlight = 4;
+		
+		public FileUpload(HttpClient httpClient)
+			: base(httpClient, MaxUploadsInFlight)
+		{
+		}
+
+		public FileUpload(HttpClient httpClient, int maxUploadsInFlight)
+			: base(httpClient, maxUploadsInFlight)
+		{
+		}
+
+		protected override async Task DoFileTransferAsync(string source, string destination)
+		{
+			using (var fileStream = new FileStream(source, FileMode.Open, FileAccess.Read))
+			{
+				using (var content = new StreamContent(fileStream))
+				{
+					await _httpClient.PutAsync(destination, content, _cancellationTokenSource.Token);
+				}
+			}
+		}
+
+		protected override Task<long> DoFileSizeAsync(string source)
+		{
+			return Task.FromResult(new FileInfo(source).Length);
+		}
+	}
+}
\ No newline at end of file
diff --git a/Engine/Source/Programs/SkeinCLI/SkeinCLI/FileTransferEventArgs.cs b/Engine/Source/Programs/SkeinCLI/SkeinCLI/FileTransferEventArgs.cs
new file mode 100644
--- /dev/null
+++ b/Engine/Source/Programs/SkeinCLI/SkeinCLI/FileTransferEventArgs.cs
@@ -0,0 +1,35 @@
+// Copyright Epic Games, Inc. All Rights Reserved.
+
+using System;
+using System.IO;
+
+namespace SkeinCLI
+{
+	public class TransferStartedEventArgs : EventArgs
+	{
+		public TransferStartedEventArgs(string source, string destination, long size)
+		{
+			Source = source;
+			Destination = destination;
+			Size = size;
+		}
+
+		public string Source { get; }
+		public string Destination { get; }
+		public long Size { get; }
+	}
+
+	public class TransferCompletedEventArgs : EventArgs
+	{
+		public TransferCompletedEventArgs(string source, string destination, bool result)
+		{
+			Source = source;
+			Destination = destination;
+			Result = result;
+		}
+
+		public string Source { get; }
+		public string Destination { get; }
+		public bool Result { get; }
+	}
+}
\ No newline at end of file
